RPCBase - module

#rnx

#Synced-CL 0

diff --git a/Engine/Plugins/RPCBase/Source/RPCBase/Private/RPCBaseModule.cpp b/Engine/Plugins/RPCBase/Source/RPCBase/Private/RPCBaseModule.cpp
new file mode 100644
--- /dev/null
+++ b/Engine/Plugins/RPCBase/Source/RPCBase/Private/RPCBaseModule.cpp
@@ -0,0 +1,14 @@
+// Copyright Epic Games, Inc. All Rights Reserved.
+
+#include "RPCBaseModule.h"
+#include "Modules/ModuleManager.h"
+
+void FRPCBaseModule::StartupModule()
+{
+}
+
+void FRPCBaseModule::ShutdownModule()
+{
+}
+
+IMPLEMENT_MODULE(FRPCBaseModule, RPCBase);
\ No newline at end of file
diff --git a/Engine/Plugins/RPCBase/Source/RPCBase/Public/RPCBaseModule.h b/Engine/Plugins/RPCBase/Source/RPCBase/Public/RPCBaseModule.h
new file mode 100644
--- /dev/null
+++ b/Engine/Plugins/RPCBase/Source/RPCBase/Public/RPCBaseModule.h
@@ -0,0 +1,13 @@
+// Copyright Epic Games, Inc. All Rights Reserved.
+
+#pragma once
+
+#include "Modules/ModuleInterface.h"
+
+class FRPCBaseModule : public IModuleInterface
+{
+public:
+	virtual void StartupModule() override;
+	virtual void ShutdownModule() override;
+};
+
